Fix non unity build

diff --git a/Engine/Source/Developer/MeshBoneReduction/Public/MeshBoneReduction.h b/Engine/Source/Developer/MeshBoneReduction/Public/MeshBoneReduction.h
--- a/Engine/Source/Developer/MeshBoneReduction/Public/MeshBoneReduction.h
+++ b/Engine/Source/Developer/MeshBoneReduction/Public/MeshBoneReduction.h
@@ -8,6 +8,7 @@
 
 class USkeletalMesh;
 struct FSkelMeshSection;
+struct FImportedSkinWeightProfileData;
 
 /**
  * Mesh reduction interface.
