Avoid assertions from re-registering components during PostLoad

Only re-register the skeletal mesh component when it has been rebuilt, instead of re-registering every component.

Fix suggested by [REDACTED_PERSON]

diff --git a/Engine/Source/Developer/MeshBoneReduction/Private/MeshBoneReduction.cpp b/Engine/Source/Developer/MeshBoneReduction/Private/MeshBoneReduction.cpp
--- a/Engine/Source/Developer/MeshBoneReduction/Private/MeshBoneReduction.cpp
+++ b/Engine/Source/Developer/MeshBoneReduction/Private/MeshBoneReduction.cpp
@@ -473,17 +473,13 @@ public:
 		NewModel->ActiveBoneIndices.Sort();
 		NewModel->RequiredBones.Sort();
 
+		// Call post edit change and re-register skeletal mesh component
 		if (bCallPostEditChange)
 		{
-			SkeletalMesh->PostEditChange();
+			FScopedSkeletalMeshPostEditChange ScopedSkeletalMeshPostEditChange(SkeletalMesh);
 		}
 		SkeletalMesh->MarkPackageDirty();
-		
-		//Reregister skinned mesh component if we call post edit change
-		if (bCallPostEditChange)
-		{
-			TComponentReregisterContext<USkinnedMeshComponent> ReregisterContext;
-		}
+
 		return true;
 	}
 };
