BakePose and BakePoseOverrides generate incorrect poses
    - Added support for retargeting when using Bake Pose with MeshBoneReduction module
    - Small const fixes and tweaks to MeshBoneReduction module
    - Removed unused array from in Skel Mesh reduction plugin

diff --git a/Engine/Plugins/Experimental/SkeletalReduction/Source/Private/SkeletalMeshReductionPlugin.cpp b/Engine/Plugins/Experimental/SkeletalReduction/Source/Private/SkeletalMeshReductionPlugin.cpp
--- a/Engine/Plugins/Experimental/SkeletalReduction/Source/Private/SkeletalMeshReductionPlugin.cpp
+++ b/Engine/Plugins/Experimental/SkeletalReduction/Source/Private/SkeletalMeshReductionPlugin.cpp
@@ -2024,12 +2024,11 @@ void FQuadricSkeletalMeshReduction::ReduceSkeletalMesh(USkeletalMesh& SkeletalMe
 	}
 
 	// if it has bake pose, gets ref to local matrices using bake pose
-	if (const UAnimSequence* BakePoseAnim = SkeletalMesh.GetLODInfo(LODIndex)->BakePose)
+	if (const UAnimSequence* BakePoseAnim = SkeletalMesh.GetBakePose(LODIndex))
 	{
 		FMemMark Mark(FMemStack::Get());
 		
 		const FReferenceSkeleton& RefSkeleton = SkeletalMesh.GetRefSkeleton();
-		const TArray<FTransform>& RefPoseInLocal = RefSkeleton.GetRefBonePose();
 
 		// Get component space retarget base pose, will be equivalent of ref-pose if not edited
 		TArray<FTransform> ComponentSpaceRefPose;
@@ -2054,6 +2053,7 @@ void FQuadricSkeletalMeshReduction::ReduceSkeletalMesh(USkeletalMesh& SkeletalMe
 
 		FCompactPose Pose;
 		Pose.SetBoneContainer(&RequiredBones);
+		Pose.ResetToRefPose();
 
 		// Retrieve animated pose from anim sequence (including retargeting)
 		const USkeleton* Skeleton = SkeletalMesh.GetSkeleton();
diff --git a/Engine/Source/Developer/MeshBoneReduction/Private/MeshBoneReduction.cpp b/Engine/Source/Developer/MeshBoneReduction/Private/MeshBoneReduction.cpp
--- a/Engine/Source/Developer/MeshBoneReduction/Private/MeshBoneReduction.cpp
+++ b/Engine/Source/Developer/MeshBoneReduction/Private/MeshBoneReduction.cpp
@@ -14,6 +14,7 @@
 #include "AnimationBlueprintLibrary.h"
 #include "Async/ParallelFor.h"
 #include "Animation/AnimSequence.h"
+#include "Animation/AnimSequenceHelpers.h"
 
 class FMeshBoneReductionModule : public IMeshBoneReductionModule
 {
@@ -271,7 +272,7 @@ public:
 		}
 	}
 
-	void RetrieveBoneMatrices(USkeletalMesh* SkeletalMesh, const int32 LODIndex, TArray<FBoneIndexType>& BonesToRemove, TArray<FMatrix>& InOutMatrices)
+	void RetrieveBoneMatrices(USkeletalMesh* SkeletalMesh, const int32 LODIndex, TArray<FBoneIndexType>& BonesToRemove, TArray<FMatrix>& InOutMatrices) const
 	{
 		if (!SkeletalMesh->IsValidLODIndex(LODIndex))
 		{
@@ -280,84 +281,102 @@ public:
 
 		// Retrieve all bone names in skeleton
 		TArray<FName> BoneNames;
-		const int32 NumBones = SkeletalMesh->GetRefSkeleton().GetRawBoneNum();
+		const int32 NumBones = SkeletalMesh->GetRefSkeleton().GetNum();
 		for (int32 BoneIndex = 0; BoneIndex < NumBones; ++BoneIndex)
 		{
 			BoneNames.Add(SkeletalMesh->GetRefSkeleton().GetBoneName(BoneIndex));
 		}
-				
-		TArray<FMatrix> MultipliedBonePoses;
-		const UAnimSequence* BakePose = SkeletalMesh->GetBakePose(LODIndex);
-		if (BakePose)
+
+		// get the relative to ref pose matrices
+		TArray<FMatrix> RelativeToRefPoseMatrices;
+		RelativeToRefPoseMatrices.AddDefaulted(NumBones);
+
+		// Set initial matrices to identity
+		for (int32 Index = 0; Index < NumBones; ++Index)
 		{
-			// Retrieve posed bone transforms
-			TArray<FTransform> BonePoses;
-			UAnimationBlueprintLibrary::GetBonePosesForFrame(BakePose, BoneNames, 0, true, BonePoses);
-			MultipliedBonePoses.AddDefaulted(BonePoses.Num());
-
-			// Retrieve ref pose bone transforms
-			TArray<FTransform> RefBonePoses = SkeletalMesh->GetRefSkeleton().GetRawRefBonePose();
-			TArray<FMatrix> MultipliedRefBonePoses;
-			MultipliedRefBonePoses.AddDefaulted(RefBonePoses.Num());
-
-			const bool bDebugBonePoses = false;
-
-			TArray<int32> Processed;
-			Processed.SetNumZeroed(RefBonePoses.Num());
-			// Multiply out parent to child transforms for ref-pose
-			for (int32 BonePoseIndex = 0; BonePoseIndex < RefBonePoses.Num(); BonePoseIndex++)
-			{
-				const int32 BoneIndex = SkeletalMesh->GetRefSkeleton().FindRawBoneIndex(BoneNames[BonePoseIndex]);
-				const int32 ParentIndex = SkeletalMesh->GetRefSkeleton().GetParentIndex(BoneIndex);
-				MultipliedRefBonePoses[BoneIndex] = RefBonePoses[BoneIndex].ToMatrixWithScale();
+			RelativeToRefPoseMatrices[Index] = FMatrix::Identity;
+		}
+
+		// if it has bake pose, gets ref to local matrices using bake pose
+		if (const UAnimSequence* BakePoseAnim = SkeletalMesh->GetBakePose(LODIndex))
+		{
+			FMemMark Mark(FMemStack::Get());
+
+			const FReferenceSkeleton& RefSkeleton = SkeletalMesh->GetRefSkeleton();
 
-				if (ParentIndex != INDEX_NONE)
-				{
-					checkf(ParentIndex == 0 || Processed[ParentIndex] == 1, TEXT("Parent bone was not yet processed"));
-					UE_CLOG(bDebugBonePoses, LogMeshBoneReduction, Log, TEXT("Original: [%i]\n%s"), BonePoseIndex, *RefBonePoses[BoneIndex].ToHumanReadableString());
-					MultipliedRefBonePoses[BoneIndex] = MultipliedRefBonePoses[BoneIndex] * MultipliedRefBonePoses[ParentIndex];
-					UE_CLOG(bDebugBonePoses, LogMeshBoneReduction, Log, TEXT("Relative: [%i]\n%s"), BonePoseIndex, *(FTransform(MultipliedRefBonePoses[BoneIndex]).ToHumanReadableString()));		
-				}		
 
-				Processed[BoneIndex] = 1;
+			// Setup BoneContainer and CompactPose
+			TArray<FBoneIndexType> RequiredBoneIndexArray;
+			RequiredBoneIndexArray.AddUninitialized(RefSkeleton.GetNum());
+			for (int32 BoneIndex = 0; BoneIndex < RequiredBoneIndexArray.Num(); ++BoneIndex)
+			{
+				RequiredBoneIndexArray[BoneIndex] = BoneIndex;
 			}
 
-			Processed.Empty();
-			Processed.SetNumZeroed(BonePoses.Num());
-			// Multiply out parent to child transforms for bake-pose
-			for (int32 BonePoseIndex = 0; BonePoseIndex < BonePoses.Num(); BonePoseIndex++)
+			FBoneContainer RequiredBones(RequiredBoneIndexArray, false, *SkeletalMesh);
+			RequiredBones.SetUseRAWData(true);
+
+			FCompactPose Pose;
+			Pose.SetBoneContainer(&RequiredBones);
+			Pose.ResetToRefPose();
+
+			// Get component space retarget base pose, will be equivalent of ref-pose if not edited
+			TArray<FTransform> ComponentSpaceRefPose;
+			FAnimationRuntime::FillUpComponentSpaceTransformsRetargetBasePose(SkeletalMesh, ComponentSpaceRefPose);
+			
+			// Retrieve animated pose from anim sequence (including retargeting)
+			const USkeleton* Skeleton = SkeletalMesh->GetSkeleton();
+			const FName RetargetSource = Skeleton->GetRetargetSourceForMesh(SkeletalMesh);
+			UE::Anim::BuildPoseFromModel(BakePoseAnim->GetDataModel(), Pose, 0.f, EAnimInterpolationType::Step, RetargetSource, Skeleton->GetRefLocalPoses(RetargetSource));
+			
+			// Calculate component space animated pose matrices
+			TArray<FMatrix> ComponentSpaceAnimatedPose;
+			ComponentSpaceAnimatedPose.AddDefaulted(NumBones);
+
+			for (int32 BoneIndex = 0; BoneIndex < NumBones; ++BoneIndex)
 			{
-				const int32 BoneIndex = SkeletalMesh->GetRefSkeleton().FindRawBoneIndex(BoneNames[BonePoseIndex]);
-				const int32 ParentIndex = SkeletalMesh->GetRefSkeleton().GetParentIndex(BoneIndex);
-				MultipliedBonePoses[BoneIndex] = BonePoses[BoneIndex].ToMatrixWithScale();
+				const FCompactPoseBoneIndex PoseBoneIndex(BoneIndex);
+				const int32 ParentIndex = RefSkeleton.GetParentIndex(BoneIndex);
 				if (ParentIndex != INDEX_NONE)
 				{
-					checkf(ParentIndex == 0 || Processed[ParentIndex] == 1, TEXT("Parent bone was not yet processed"));
-					UE_CLOG(bDebugBonePoses, LogMeshBoneReduction, Log, TEXT("Original: [%i]\n%s"), BonePoseIndex, *BonePoses[BoneIndex].ToHumanReadableString());					
-					MultipliedBonePoses[BoneIndex] = MultipliedBonePoses[BoneIndex] * MultipliedBonePoses[ParentIndex];
-					UE_CLOG(bDebugBonePoses, LogMeshBoneReduction, Log, TEXT("Relative: [%i]\n%s"), BonePoseIndex, *(FTransform(MultipliedBonePoses[BoneIndex]).ToHumanReadableString()));			
+					// If the bone will be removed, get the local-space retarget-ed animation bone transform
+					if (BonesToRemove.Contains(BoneIndex))
+					{
+						ComponentSpaceAnimatedPose[BoneIndex] = Pose[PoseBoneIndex].ToMatrixWithScale() * ComponentSpaceAnimatedPose[ParentIndex];
+					}
+					// Otherwise use the component-space retarget base pose transform
+					else
+					{
+						ComponentSpaceAnimatedPose[BoneIndex] = ComponentSpaceRefPose[BoneIndex].ToMatrixWithScale();
+					}
+				}
+				else
+				{
+					// If the bone will be removed, get the retarget-ed animation bone transform
+					if (BonesToRemove.Contains(BoneIndex))
+					{
+						ComponentSpaceAnimatedPose[BoneIndex] = Pose[PoseBoneIndex].ToMatrixWithScale();
+					}
+					// Otherwise use the retarget base pose transform
+					else
+					{
+						ComponentSpaceAnimatedPose[BoneIndex] = ComponentSpaceRefPose[BoneIndex].ToMatrixWithScale();
+					}
 				}
-
-				Processed[BoneIndex] = 1;
 			}
 
-			// Calculate final bone pose transforms from ref-pose to bake-pose 
+			// Calculate relative to retarget base (ref) pose matrix
 			for (int32 BoneIndex = 0; BoneIndex < NumBones; ++BoneIndex)
 			{
-				MultipliedBonePoses[BoneIndex] = MultipliedRefBonePoses[BoneIndex].Inverse() * MultipliedBonePoses[BoneIndex];
-				UE_CLOG(bDebugBonePoses, LogMeshBoneReduction, Log, TEXT("Final: [%i]\n%s"), BoneIndex, *(FTransform(MultipliedBonePoses[BoneIndex]).ToHumanReadableString()));
+				RelativeToRefPoseMatrices[BoneIndex] = ComponentSpaceRefPose[BoneIndex].ToMatrixWithScale().Inverse() * ComponentSpaceAnimatedPose[BoneIndex];
 			}
 		}
-		else
-		{
-			MultipliedBonePoses.AddDefaulted(BoneNames.Num());
-		}
-
+		
 		// Add bone transforms we're interested in
 		InOutMatrices.Reset(BonesToRemove.Num());
 		for (const FBoneIndexType& Index : BonesToRemove)
 		{
-			InOutMatrices.Add(MultipliedBonePoses[Index]);
+			InOutMatrices.Add(RelativeToRefPoseMatrices[Index]);
 		}
 	}
 
@@ -408,7 +427,7 @@ public:
 			{
 				for (const FBoneReference& BoneReference : SkeletalMesh->GetLODInfo(DesiredLOD)->BonesToRemove)
 				{
-					int32 BoneIndex = SkeletalMesh->GetRefSkeleton().FindRawBoneIndex(BoneReference.BoneName);
+					const int32 BoneIndex = SkeletalMesh->GetRefSkeleton().FindRawBoneIndex(BoneReference.BoneName);
 					if (BoneIndex != INDEX_NONE)
 					{
 						BoneIndices.AddUnique(BoneIndex);
@@ -437,25 +456,28 @@ public:
 					{
 						FVector TangentX = Vertex.TangentX;
 						FVector TangentY = Vertex.TangentY;
-						FVector TangentZ = (FVector4)Vertex.TangentZ;
+						FVector TangentZ = Vertex.TangentZ;
 						FVector Position = Vertex.Position;
 						for (uint8 InfluenceIndex = 0; InfluenceIndex < MAX_TOTAL_INFLUENCES; ++InfluenceIndex)
-						{
-							const int32 ArrayIndex = BoneIndices.IndexOfByKey(Section.BoneMap[Vertex.InfluenceBones[InfluenceIndex]]);
-							if (ArrayIndex != INDEX_NONE)
+				        {
+							if (Vertex.InfluenceWeights[InfluenceIndex] > 0)
 							{
-								Position += ((RemovedBoneMatrices[ArrayIndex].TransformPosition(Vertex.Position) - Vertex.Position) * ((float)Vertex.InfluenceWeights[InfluenceIndex] * InfluenceMultiplier));
+								const int32 ArrayIndex = BoneIndices.IndexOfByKey(Section.BoneMap[Vertex.InfluenceBones[InfluenceIndex]]);
+								if (ArrayIndex != INDEX_NONE)
+								{
+									Position += ((RemovedBoneMatrices[ArrayIndex].TransformPosition(Vertex.Position) - Vertex.Position) * ((float)Vertex.InfluenceWeights[InfluenceIndex] * InfluenceMultiplier));
 
-								TangentX += ((RemovedBoneMatrices[ArrayIndex].TransformVector(Vertex.TangentX) - Vertex.TangentX) * ((float)Vertex.InfluenceWeights[InfluenceIndex] * InfluenceMultiplier));
-								TangentY += ((RemovedBoneMatrices[ArrayIndex].TransformVector(Vertex.TangentY) - Vertex.TangentY) * ((float)Vertex.InfluenceWeights[InfluenceIndex] * InfluenceMultiplier));
-								TangentZ += ((RemovedBoneMatrices[ArrayIndex].TransformVector((FVector4)Vertex.TangentZ) - Vertex.TangentZ) * ((float)Vertex.InfluenceWeights[InfluenceIndex] * InfluenceMultiplier));
+									TangentX += ((RemovedBoneMatrices[ArrayIndex].TransformVector(Vertex.TangentX) - Vertex.TangentX) * ((float)Vertex.InfluenceWeights[InfluenceIndex] * InfluenceMultiplier));
+									TangentY += ((RemovedBoneMatrices[ArrayIndex].TransformVector(Vertex.TangentY) - Vertex.TangentY) * ((float)Vertex.InfluenceWeights[InfluenceIndex] * InfluenceMultiplier));
+									TangentZ += ((RemovedBoneMatrices[ArrayIndex].TransformVector(Vertex.TangentZ) - Vertex.TangentZ) * ((float)Vertex.InfluenceWeights[InfluenceIndex] * InfluenceMultiplier));
+								}
 							}
-						}
+				        }
 
 						Vertex.Position = Position;
 						Vertex.TangentX = TangentX.GetSafeNormal();
 						Vertex.TangentY = TangentY.GetSafeNormal();
-						uint8 WComponent = Vertex.TangentZ.W;
+						const uint8 WComponent = Vertex.TangentZ.W;
 						Vertex.TangentZ = TangentZ.GetSafeNormal();
 						Vertex.TangentZ.W = WComponent;
 					}
