UE-113768
[CrashReport] UE4Editor-ContentBrowser!TBaseSPMethodDelegateInstance<0,SContentBrowser,0,void ,FDefaultDelegateUserPolicy>::Execute(TArrayView<FContentBrowserItem const ,int>,EAssetTypeActivationMethod::Type) [DelegateInstancesImpl.h:290]

diff --git a/Engine/Source/Developer/MeshBoneReduction/Private/MeshBoneReduction.cpp b/Engine/Source/Developer/MeshBoneReduction/Private/MeshBoneReduction.cpp
--- a/Engine/Source/Developer/MeshBoneReduction/Private/MeshBoneReduction.cpp
+++ b/Engine/Source/Developer/MeshBoneReduction/Private/MeshBoneReduction.cpp
@@ -382,8 +382,15 @@ public:
 
 	bool ReduceBoneCounts(USkeletalMesh* SkeletalMesh, int32 DesiredLOD, const TArray<FName>* BoneNamesToRemove, bool bCallPostEditChange /*= true*/) override
 	{
-		if (!ensure(SkeletalMesh))
+		if (!SkeletalMesh)
+		{
+			UE_LOG(LogMeshBoneReduction, Error, TEXT("Failed to remove Skeletal Mesh LOD %i bones, as the Skeletal Mesh is invalid."), DesiredLOD);
+			return false;
+		}
+
+		if(!SkeletalMesh->GetSkeleton())
 		{
+			UE_LOG(LogMeshBoneReduction, Error, TEXT("Failed to remove bones from LOD %i in %s, as its Skeleton is invalid."), DesiredLOD, *SkeletalMesh->GetPathName());
 			return false;
 		}
 
