MeshBoneReduction - truncation warning fixes

diff --git a/Engine/Source/Developer/MeshBoneReduction/Private/MeshBoneReduction.cpp b/Engine/Source/Developer/MeshBoneReduction/Private/MeshBoneReduction.cpp
--- a/Engine/Source/Developer/MeshBoneReduction/Private/MeshBoneReduction.cpp
+++ b/Engine/Source/Developer/MeshBoneReduction/Private/MeshBoneReduction.cpp
@@ -17,6 +17,7 @@
 #include "Async/ParallelFor.h"
 #include "Animation/AnimSequence.h"
 #include "Animation/AnimSequenceHelpers.h"
+#include "Templates/UnrealTemplate.h"
 
 class FMeshBoneReductionModule : public IMeshBoneReductionModule
 {
@@ -40,10 +41,10 @@ public:
 	{
 	}
 
-	void EnsureChildrenPresents(FBoneIndexType BoneIndex, const TArray<FMeshBoneInfo>& RefBoneInfo, TArray<FBoneIndexType>& OutBoneIndicesToRemove)
+	void EnsureChildrenPresents(int32 BoneIndex, const TArray<FMeshBoneInfo>& RefBoneInfo, TArray<int32>& OutBoneIndicesToRemove)
 	{
 		// just look for direct parent, we could look for RefBoneInfo->Ischild, but more expensive, and no reason to do that all the work
-		for (int32 ChildBoneIndex = 0; ChildBoneIndex < RefBoneInfo.Num(); ++ChildBoneIndex)
+		for (int32 ChildBoneIndex = 0; ChildBoneIndex < static_cast<int32>( RefBoneInfo.Num()); ++ChildBoneIndex)
 		{
 			if (RefBoneInfo[ChildBoneIndex].ParentIndex == BoneIndex)
 			{
@@ -66,7 +67,7 @@ public:
 		}
 
 		const TArray<FMeshBoneInfo> & RefBoneInfo = SkeletalMesh->GetRefSkeleton().GetRawRefBoneInfo();
-		TArray<FBoneIndexType> BoneIndicesToRemove;
+		TArray<int32> BoneIndicesToRemove;
 
 		// originally this code was accumulating from LOD 0->DesiredLOd, but that should be done outside of tool if they want to
 		// removing it, and just include DesiredLOD
@@ -127,7 +128,7 @@ public:
 				ParentIndex = RefBoneInfo[ParentIndex].ParentIndex;
 			}
 
-			OutBonesToReplace.Add(BoneIndex, ParentIndex);
+			OutBonesToReplace.Add(IntCastChecked<FBoneIndexType>(BoneIndex), IntCastChecked<FBoneIndexType>(ParentIndex));
 		}
 
 		return ( OutBonesToReplace.Num() > 0 );
@@ -274,7 +275,7 @@ public:
 		}
 	}
 
-	void RetrieveBoneMatrices(USkeletalMesh* SkeletalMesh, const int32 LODIndex, TArray<FBoneIndexType>& BonesToRemove, TArray<FMatrix>& InOutMatrices) const
+	void RetrieveBoneMatrices(USkeletalMesh* SkeletalMesh, const int32 LODIndex, const TArray<int32>& BonesToRemove, TArray<FMatrix>& InOutMatrices) const
 	{
 		if (!SkeletalMesh->IsValidLODIndex(LODIndex))
 		{
@@ -310,11 +311,13 @@ public:
 			// Setup BoneContainer and CompactPose
 			TArray<FBoneIndexType> RequiredBoneIndexArray;
 			RequiredBoneIndexArray.AddUninitialized(RefSkeleton.GetNum());
-			for (int32 BoneIndex = 0; BoneIndex < RequiredBoneIndexArray.Num(); ++BoneIndex)
 			{
-				RequiredBoneIndexArray[BoneIndex] = BoneIndex;
+				FBoneIndexType RequiredBoneIndexNum = IntCastChecked<FBoneIndexType>(RequiredBoneIndexArray.Num());
+				for (FBoneIndexType BoneIndex = 0; BoneIndex < RequiredBoneIndexNum; ++BoneIndex)
+				{
+					RequiredBoneIndexArray[BoneIndex] = BoneIndex;
+				}
 			}
-
 			FBoneContainer RequiredBones(RequiredBoneIndexArray, false, *SkeletalMesh);
 			RequiredBones.SetUseRAWData(true);
 
@@ -382,7 +385,7 @@ public:
 		
 		// Add bone transforms we're interested in
 		InOutMatrices.Reset(BonesToRemove.Num());
-		for (const FBoneIndexType& Index : BonesToRemove)
+		for (const int32 Index : BonesToRemove)
 		{
 			InOutMatrices.Add(RelativeToRefPoseMatrices[Index]);
 		}
@@ -435,7 +438,7 @@ public:
 
 			FSkeletalMeshLODModel::CopyStructure(NewModel, SrcModel);
 
-			TArray<FBoneIndexType> BoneIndices;
+			TArray<int32> BoneIndices;
 			TArray<FMatrix> RemovedBoneMatrices;
 			const bool bBakePoseToRemovedInfluences = (SkeletalMesh->GetBakePose(DesiredLOD) != nullptr);
 			if (bBakePoseToRemovedInfluences)
@@ -493,7 +496,7 @@ public:
 						Vertex.Position = (FVector3f)Position;
 						Vertex.TangentX = (FVector3f)TangentX.GetSafeNormal();
 						Vertex.TangentY = (FVector3f)TangentY.GetSafeNormal();
-						const uint8 WComponent = Vertex.TangentZ.W;
+						const uint8 WComponent = static_cast<uint8>(Vertex.TangentZ.W);
 						Vertex.TangentZ = (FVector3f)TangentZ.GetSafeNormal();
 						Vertex.TangentZ.W = WComponent;
 					}
