Fix a crash when removing LOD bones.
#rnx

diff --git a/Engine/Source/Developer/MeshBoneReduction/Private/MeshBoneReduction.cpp b/Engine/Source/Developer/MeshBoneReduction/Private/MeshBoneReduction.cpp
--- a/Engine/Source/Developer/MeshBoneReduction/Private/MeshBoneReduction.cpp
+++ b/Engine/Source/Developer/MeshBoneReduction/Private/MeshBoneReduction.cpp
@@ -267,8 +267,11 @@ public:
 					for (auto Kvp : SkinWeightProfiles)
 					{
 						FImportedSkinWeightProfileData& SkinWeightProfile = SkinWeightProfiles.FindChecked(Kvp.Key);
-						check(SkinWeightProfile.SkinWeights.IsValidIndex(RealVertexIndex));
-						RemapBoneInfluenceVertexIndex(SkinWeightProfile.SkinWeights[RealVertexIndex].InfluenceBones, SkinWeightProfile.SkinWeights[RealVertexIndex].InfluenceWeights);
+						//If the vertex index is not handle by the runtime skin weight profile, the next skeletal mesh build will fix the runtime skin weight profile amount of vertices.
+						if (SkinWeightProfile.SkinWeights.IsValidIndex(RealVertexIndex))
+						{
+							RemapBoneInfluenceVertexIndex(SkinWeightProfile.SkinWeights[RealVertexIndex].InfluenceBones, SkinWeightProfile.SkinWeights[RealVertexIndex].InfluenceWeights);
+						}
 					}
 				}
 			}
